Attempt to blindly fix CIS.

diff --git a/Engine/Source/Runtime/MovieSceneCore/Public/MovieSceneInstance.h b/Engine/Source/Runtime/MovieSceneCore/Public/MovieSceneInstance.h
--- a/Engine/Source/Runtime/MovieSceneCore/Public/MovieSceneInstance.h
+++ b/Engine/Source/Runtime/MovieSceneCore/Public/MovieSceneInstance.h
@@ -2,6 +2,8 @@
 
 #pragma once
 
+class UMovieSceneTrack;
+
 typedef TMap< TWeakObjectPtr<UMovieSceneTrack>, TSharedPtr<class IMovieSceneTrackInstance> > FMovieSceneInstanceMap;
 class FMovieSceneInstance : public TSharedFromThis<FMovieSceneInstance>
 {
diff --git a/Engine/Source/Runtime/UMG/Public/Animation/UMGSequencePlayer.h b/Engine/Source/Runtime/UMG/Public/Animation/UMGSequencePlayer.h
--- a/Engine/Source/Runtime/UMG/Public/Animation/UMGSequencePlayer.h
+++ b/Engine/Source/Runtime/UMG/Public/Animation/UMGSequencePlayer.h
@@ -3,6 +3,8 @@
 #pragma once
 
 #include "IMovieScenePlayer.h"
+#include "MovieScene.h"
+#include "MovieSceneBindings.h"
 #include "UMGSequencePlayer.generated.h"
 
 UCLASS(transient)
