UMG - Animations now only evaluate the parts of multi-component curves (e.g, margin, color, vector) which have keys.  If a curve has no keys the current value of the field is not changed

diff --git a/Engine/Source/Developer/MovieSceneTools/Private/PropertyTrackEditor.cpp b/Engine/Source/Developer/MovieSceneTools/Private/PropertyTrackEditor.cpp
--- a/Engine/Source/Developer/MovieSceneTools/Private/PropertyTrackEditor.cpp
+++ b/Engine/Source/Developer/MovieSceneTools/Private/PropertyTrackEditor.cpp
@@ -134,7 +134,10 @@ private:
 					}
 				}
 
-				if (InsertKeyIndex == INDEX_NONE && bShouldAddKey) {InsertKeyIndex = TimesWithKeys.Num();}
+				if (InsertKeyIndex == INDEX_NONE && bShouldAddKey)
+				{
+					InsertKeyIndex = TimesWithKeys.Num();
+				}
 
 				if (bShouldAddKey)
 				{
@@ -146,7 +149,7 @@ private:
 		// @todo Sequencer Optimize - This another O(n^2) loop, since Eval is O(n)!
 		for (int32 i = 0; i < TimesWithKeys.Num(); ++i)
 		{
-			OutColorKeys.Add(TKeyValuePair<float, FLinearColor>(TimesWithKeys[i], Section->Eval(TimesWithKeys[i])));
+			OutColorKeys.Add( TKeyValuePair<float, FLinearColor>(TimesWithKeys[i], Section->Eval(TimesWithKeys[i], FLinearColor(0,0,0,0) ) ) );
 		}
 	}
 };
@@ -496,7 +499,7 @@ void FPropertyTrackEditor::OnKeyProperty( float KeyTime, FKeyPropertyParams KeyP
 			{
 				TrackType* TypedTrack = CastChecked<TrackType>(Track);
 
-				TypedTrack->SetPropertyName( PropertyName );
+				TypedTrack->SetPropertyNameAndPath( PropertyName, KeyPropertyParams.PropertyPath );
 				// Find or add a new section at the auto-key time and changing the property same property
 				// AddKeyToSection is not actually a virtual, it's redefined in each class with a different type
 				bool bSuccessfulAdd = TypedTrack->AddKeyToSection( KeyTime, Value );
diff --git a/Engine/Source/Developer/MovieSceneTools/Private/TransformTrackEditor.cpp b/Engine/Source/Developer/MovieSceneTools/Private/TransformTrackEditor.cpp
--- a/Engine/Source/Developer/MovieSceneTools/Private/TransformTrackEditor.cpp
+++ b/Engine/Source/Developer/MovieSceneTools/Private/TransformTrackEditor.cpp
@@ -256,7 +256,8 @@ void FTransformTrackEditor::OnTransformChangedInternals(float KeyTime, UObject*
 	{
 		UMovieSceneTrack* Track = GetTrackForObject( ObjectHandle, UMovieSceneTransformTrack::StaticClass(), Transform );
 		UMovieSceneTransformTrack* TransformTrack = CastChecked<UMovieSceneTransformTrack>( Track );
-		TransformTrack->SetPropertyName( Transform );
+		// Transform name and path are the same
+		TransformTrack->SetPropertyNameAndPath( Transform, Transform.ToString() );
 	
 		if (!TransformPair.LastTransformData.IsValid())
 		{
diff --git a/Engine/Source/Editor/PropertyEditor/Private/CategoryPropertyNode.cpp b/Engine/Source/Editor/PropertyEditor/Private/CategoryPropertyNode.cpp
--- a/Engine/Source/Editor/PropertyEditor/Private/CategoryPropertyNode.cpp
+++ b/Engine/Source/Editor/PropertyEditor/Private/CategoryPropertyNode.cpp
@@ -107,15 +107,26 @@ void FCategoryPropertyNode::InitChildNodes()
 /**
  * Appends my path, including an array index (where appropriate)
  */
-void FCategoryPropertyNode::GetQualifiedName( FString& PathPlusIndex, const bool bWithArrayIndex ) const
+bool FCategoryPropertyNode::GetQualifiedName( FString& PathPlusIndex, const bool bWithArrayIndex, const FPropertyNode* StopParent, bool bIgnoreCategories ) const
 {
-	if( ParentNode )
+	bool bAddedAnything = false;
+
+	if( ParentNode && StopParent != ParentNode )
 	{
-		ParentNode->GetQualifiedName(PathPlusIndex, bWithArrayIndex);
-		PathPlusIndex += TEXT(".");
+		bAddedAnything = ParentNode->GetQualifiedName(PathPlusIndex, bWithArrayIndex, StopParent, bIgnoreCategories );
+		if( bAddedAnything )
+		{
+			PathPlusIndex += TEXT(".");
+		}
 	}
 	
-	GetCategoryName().AppendString(PathPlusIndex);
+	if( !bIgnoreCategories )
+	{
+		bAddedAnything = true;
+		GetCategoryName().AppendString(PathPlusIndex);
+	}
+
+	return bAddedAnything;
 }
 
 
diff --git a/Engine/Source/Editor/PropertyEditor/Private/CategoryPropertyNode.h b/Engine/Source/Editor/PropertyEditor/Private/CategoryPropertyNode.h
--- a/Engine/Source/Editor/PropertyEditor/Private/CategoryPropertyNode.h
+++ b/Engine/Source/Editor/PropertyEditor/Private/CategoryPropertyNode.h
@@ -41,7 +41,7 @@ protected:
 	/**
 	 * Appends my path, including an array index (where appropriate)
 	 */
-	virtual void GetQualifiedName( FString& PathPlusIndex, const bool bWithArrayIndex ) const override;
+	virtual bool GetQualifiedName( FString& PathPlusIndex, const bool bWithArrayIndex, const FPropertyNode* StopParent = nullptr, bool bIgnoreCategories = false ) const override;
 
 	FString GetSubcategoryName() const;
 	/**
diff --git a/Engine/Source/Editor/PropertyEditor/Private/ObjectPropertyNode.cpp b/Engine/Source/Editor/PropertyEditor/Private/ObjectPropertyNode.cpp
--- a/Engine/Source/Editor/PropertyEditor/Private/ObjectPropertyNode.cpp
+++ b/Engine/Source/Editor/PropertyEditor/Private/ObjectPropertyNode.cpp
@@ -511,14 +511,22 @@ TSharedPtr<FPropertyNode> FObjectPropertyNode::GenerateSingleChild( FName ChildP
 /**
  * Appends my path, including an array index (where appropriate)
  */
-void FObjectPropertyNode::GetQualifiedName(FString& PathPlusIndex, const bool bWithArrayIndex) const
+bool FObjectPropertyNode::GetQualifiedName(FString& PathPlusIndex, bool bWithArrayIndex, const FPropertyNode* StopParent, bool bIgnoreCategories ) const
 {
-	if( ParentNode )
+	bool bAddedAnything = false;
+	if( ParentNode && ParentNode != StopParent )
 	{
-		ParentNode->GetQualifiedName(PathPlusIndex, bWithArrayIndex);
-		PathPlusIndex += TEXT(".");
+		bAddedAnything = ParentNode->GetQualifiedName(PathPlusIndex, bWithArrayIndex, StopParent, bIgnoreCategories);
+		if( bAddedAnything )
+		{
+			PathPlusIndex += TEXT(".");
+		}
 	}
+
+	bAddedAnything = true;
 	PathPlusIndex += TEXT("Object");
+
+	return bAddedAnything;
 }
 
 // Looks at the Objects array and returns the best base class.  Called by
diff --git a/Engine/Source/Editor/PropertyEditor/Private/ObjectPropertyNode.h b/Engine/Source/Editor/PropertyEditor/Private/ObjectPropertyNode.h
--- a/Engine/Source/Editor/PropertyEditor/Private/ObjectPropertyNode.h
+++ b/Engine/Source/Editor/PropertyEditor/Private/ObjectPropertyNode.h
@@ -101,7 +101,7 @@ protected:
 	/** FPropertyNode interface */
 	virtual void InitBeforeNodeFlags() override;
 	virtual void InitChildNodes() override;
-	virtual void GetQualifiedName( FString& PathPlusIndex, const bool bWithArrayIndex ) const override;
+	virtual bool GetQualifiedName( FString& PathPlusIndex, const bool bWithArrayIndex, const FPropertyNode* StopParent = nullptr, bool bIgnoreCategories = false ) const override;
 
 	/**
 	 * Looks at the Objects array and creates the best base class.  Called by
diff --git a/Engine/Source/Editor/PropertyEditor/Private/PropertyHandleImpl.cpp b/Engine/Source/Editor/PropertyEditor/Private/PropertyHandleImpl.cpp
--- a/Engine/Source/Editor/PropertyEditor/Private/PropertyHandleImpl.cpp
+++ b/Engine/Source/Editor/PropertyEditor/Private/PropertyHandleImpl.cpp
@@ -1518,6 +1518,20 @@ bool FPropertyHandleBase::IsCustomized() const
 	return Implementation->GetPropertyNode()->HasNodeFlags( EPropertyNodeFlags::IsCustomized ) != 0;
 }
 
+FString FPropertyHandleBase::GeneratePathToProperty() const
+{
+	FString OutPath;
+	if( Implementation.IsValid() && Implementation->GetPropertyNode().IsValid() )
+	{
+		const bool bArrayIndex = true;
+		const bool bIgnoreCategories = true;
+		FPropertyNode* StopParent = Implementation->GetPropertyNode()->FindObjectItemParent();
+		Implementation->GetPropertyNode()->GetQualifiedName( OutPath, bArrayIndex, StopParent, bIgnoreCategories );
+	}
+
+	return OutPath;
+
+}
 
 TSharedRef<SWidget> FPropertyHandleBase::CreatePropertyNameWidget( const FString& NameOverride, bool bDisplayResetToDefault, bool bDisplayText, bool bDisplayThumbnail ) const
 {
diff --git a/Engine/Source/Editor/PropertyEditor/Private/PropertyHandleImpl.h b/Engine/Source/Editor/PropertyEditor/Private/PropertyHandleImpl.h
--- a/Engine/Source/Editor/PropertyEditor/Private/PropertyHandleImpl.h
+++ b/Engine/Source/Editor/PropertyEditor/Private/PropertyHandleImpl.h
@@ -373,6 +373,7 @@ public:
 	virtual FText GetResetToDefaultLabel() const override;
 	virtual void MarkHiddenByCustomization() override;
 	virtual bool IsCustomized() const override;
+	virtual FString GeneratePathToProperty() const override;
 	virtual TSharedRef<SWidget> CreatePropertyNameWidget( const FString& NameOverride = TEXT(""), bool bDisplayResetToDefault = false, bool bDisplayText = true, bool bDisplayThumbnail = true ) const override;
 	virtual TSharedRef<SWidget> CreatePropertyValueWidget() const override;
 	virtual bool IsEditConst() const override;
diff --git a/Engine/Source/Editor/PropertyEditor/Private/PropertyNode.cpp b/Engine/Source/Editor/PropertyEditor/Private/PropertyNode.cpp
--- a/Engine/Source/Editor/PropertyEditor/Private/PropertyNode.cpp
+++ b/Engine/Source/Editor/PropertyEditor/Private/PropertyNode.cpp
@@ -567,25 +567,33 @@ bool FPropertyNode::IsEditConst() const
 /**
  * Appends my path, including an array index (where appropriate)
  */
-void FPropertyNode::GetQualifiedName( FString& PathPlusIndex, const bool bWithArrayIndex ) const
+bool FPropertyNode::GetQualifiedName( FString& PathPlusIndex, const bool bWithArrayIndex, const FPropertyNode* StopParent, bool bIgnoreCategories ) const
 {
-	if( ParentNodeWeakPtr.IsValid() )
+	bool bAddedAnything = false;
+	if( ParentNodeWeakPtr.IsValid() && StopParent != ParentNode )
 	{
-		ParentNode->GetQualifiedName(PathPlusIndex, bWithArrayIndex);
-		PathPlusIndex += TEXT(".");
+		bAddedAnything = ParentNode->GetQualifiedName(PathPlusIndex, bWithArrayIndex, StopParent, bIgnoreCategories);
+		if( bAddedAnything )
+		{
+			PathPlusIndex += TEXT(".");
+		}
 	}
 
 	if( Property.IsValid() )
 	{
+		bAddedAnything = true;
 		Property->AppendName(PathPlusIndex);
 	}
 
 	if ( bWithArrayIndex && (ArrayIndex != INDEX_NONE) )
 	{
+		bAddedAnything = true;
 		PathPlusIndex += TEXT("[");
 		PathPlusIndex.AppendInt(ArrayIndex);
 		PathPlusIndex += TEXT("]");
 	}
+
+	return bAddedAnything;
 }
 
 bool FPropertyNode::GetReadAddressUncached( FPropertyNode& InPropertyNode,
diff --git a/Engine/Source/Editor/PropertyEditor/Private/PropertyNode.h b/Engine/Source/Editor/PropertyEditor/Private/PropertyNode.h
--- a/Engine/Source/Editor/PropertyEditor/Private/PropertyNode.h
+++ b/Engine/Source/Editor/PropertyEditor/Private/PropertyNode.h
@@ -359,8 +359,10 @@ public:
 	 * Gets the full name of this node
 	 * @param PathPlusIndex - return value with full path of node
 	 * @param bWithArrayIndex - If True, adds an array index (where appropriate)
+	 * @param StopParent	- Stop at this parent (if any). Does NOT include it in the path
+	 * @param bIgnoreCategories - Skip over categories
 	 */
-	virtual void GetQualifiedName( FString& PathPlusIndex, const bool bWithArrayIndex ) const;
+	virtual bool GetQualifiedName( FString& PathPlusIndex, const bool bWithArrayIndex, const FPropertyNode* StopParent = nullptr, bool bIgnoreCategories = false ) const;
 
 	// The bArrayPropertiesCanDifferInSize flag is an override for array properties which want to display
 	// e.g. the "Clear" and "Empty" buttons, even though the array properties may differ in the number of elements.
diff --git a/Engine/Source/Editor/PropertyEditor/Private/StructurePropertyNode.h b/Engine/Source/Editor/PropertyEditor/Private/StructurePropertyNode.h
--- a/Engine/Source/Editor/PropertyEditor/Private/StructurePropertyNode.h
+++ b/Engine/Source/Editor/PropertyEditor/Private/StructurePropertyNode.h
@@ -99,9 +99,10 @@ protected:
 	/** FPropertyNode interface */
 	virtual void InitChildNodes() override;
 
-	virtual void GetQualifiedName(FString& PathPlusIndex, const bool bWithArrayIndex) const override
+	virtual bool GetQualifiedName(FString& PathPlusIndex, const bool bWithArrayIndex, const FPropertyNode* StopParent = nullptr, bool bIgnoreCategories = false) const override
 	{
 		PathPlusIndex += TEXT("Struct");
+		return true;
 	}
 
 private:
diff --git a/Engine/Source/Editor/PropertyEditor/Public/PropertyHandle.h b/Engine/Source/Editor/PropertyEditor/Public/PropertyHandle.h
--- a/Engine/Source/Editor/PropertyEditor/Public/PropertyHandle.h
+++ b/Engine/Source/Editor/PropertyEditor/Public/PropertyHandle.h
@@ -326,6 +326,13 @@ public:
 	 */
 	virtual bool IsCustomized() const = 0;
 
+	/**
+	 * Generates a path from the parent UObject class to this property
+	 *
+	 * @return The path to this property
+	 */
+	virtual FString GeneratePathToProperty() const = 0;
+
 	/**
 	 * Creates a name widget for this property
 	 * @param NameOverride			The name override to use instead of the property name
diff --git a/Engine/Source/Editor/Sequencer/Private/SequencerObjectChangeListener.cpp b/Engine/Source/Editor/Sequencer/Private/SequencerObjectChangeListener.cpp
--- a/Engine/Source/Editor/Sequencer/Private/SequencerObjectChangeListener.cpp
+++ b/Engine/Source/Editor/Sequencer/Private/SequencerObjectChangeListener.cpp
@@ -48,20 +48,24 @@ void FSequencerObjectChangeListener::OnPropertyChanged( const TArray<UObject*>&
 	Params.InnerStructPropertyName = InnerStructPropName;
 	Params.bRequireAutoKey = bRequireAutoKey;
 	Params.ObjectsThatChanged = ChangedObjects;
+
 	if(StructProperty)
 	{
 		Params.PropertyHandle = &PropertyHandle;
+		Params.PropertyPath = PropertyHandle.GeneratePathToProperty();
 		ClassToPropertyChangedMap.FindRef(StructProperty->Struct->GetFName()).Broadcast( Params );
 	}
 	else if(ParentStructProperty)
 	{
-		Params.PropertyHandle = ParentHandle.Get();;
+		Params.PropertyHandle = ParentHandle.Get();
+		Params.PropertyPath = ParentHandle->GeneratePathToProperty();
 		// If the property parent is a struct, see if this property parent can be keyed. (e.g R,G,B,A for a color)
 		ClassToPropertyChangedMap.FindRef(ParentStructProperty->Struct->GetFName()).Broadcast( Params );
 	}
 	else
 	{
 		Params.PropertyHandle = &PropertyHandle;
+		Params.PropertyPath = PropertyHandle.GeneratePathToProperty();
 		// the property in question is not a struct or an inner of the struct. See if it is directly keyable
 		ClassToPropertyChangedMap.FindRef(Property->GetClass()->GetFName()).Broadcast( Params );
 	}
diff --git a/Engine/Source/Editor/Sequencer/Public/ISequencerObjectChangeListener.h b/Engine/Source/Editor/Sequencer/Public/ISequencerObjectChangeListener.h
--- a/Engine/Source/Editor/Sequencer/Public/ISequencerObjectChangeListener.h
+++ b/Engine/Source/Editor/Sequencer/Public/ISequencerObjectChangeListener.h
@@ -7,8 +7,9 @@ class IPropertyHandle;
 struct FKeyPropertyParams
 {
 	const IPropertyHandle* PropertyHandle;
-	FName InnerStructPropertyName;
 	TArray<UObject*> ObjectsThatChanged;
+	FString PropertyPath;
+	FName InnerStructPropertyName;
 	bool bRequireAutoKey;
 };
 
diff --git a/Engine/Source/Editor/UMGEditor/Private/Animation/MarginTrackEditor.cpp b/Engine/Source/Editor/UMGEditor/Private/Animation/MarginTrackEditor.cpp
--- a/Engine/Source/Editor/UMGEditor/Private/Animation/MarginTrackEditor.cpp
+++ b/Engine/Source/Editor/UMGEditor/Private/Animation/MarginTrackEditor.cpp
@@ -106,7 +106,7 @@ void FMarginTrackEditor::OnKeyMargin( float KeyTime, const FKeyPropertyParams* P
 			if( ensure( Track ) )
 			{
 				UMovieSceneMarginTrack* MarginTrack = CastChecked<UMovieSceneMarginTrack>(Track);
-				MarginTrack->SetPropertyName( PropertyName );
+				MarginTrack->SetPropertyNameAndPath( PropertyName, PropertyKeyParams->PropertyPath );
 				// Find or add a new section at the auto-key time and changing the property same property
 				// AddKeyToSection is not actually a virtual, it's redefined in each class with a different type
 				bool bSuccessfulAdd = MarginTrack->AddKeyToSection( KeyTime, Key );
diff --git a/Engine/Source/Runtime/Engine/Classes/Curves/CurveBase.h b/Engine/Source/Runtime/Engine/Classes/Curves/CurveBase.h
--- a/Engine/Source/Runtime/Engine/Classes/Curves/CurveBase.h
+++ b/Engine/Source/Runtime/Engine/Classes/Curves/CurveBase.h
@@ -352,7 +352,7 @@ public:
 	void Reset();
 
 	/** Evaluate this rich curve at the specified time */
-	float Eval(float InTime) const;
+	float Eval(float InTime, float DefaultValue = 0.0f) const;
 
 	/** Auto set tangents for any 'auto' keys in curve */
 	void AutoSetTangents(float Tension = 0.f);
diff --git a/Engine/Source/Runtime/Engine/Private/CurveBase.cpp b/Engine/Source/Runtime/Engine/Private/CurveBase.cpp
--- a/Engine/Source/Runtime/Engine/Private/CurveBase.cpp
+++ b/Engine/Source/Runtime/Engine/Private/CurveBase.cpp
@@ -680,14 +680,14 @@ static float BezierInterp2(float P0, float Y1, float Y2, float P3, float mu)
 	return Result;
 }
 
-float FRichCurve::Eval( const float InTime ) const
+float FRichCurve::Eval( const float InTime, float DefaultValue ) const
 {
 	const int32 NumKeys = Keys.Num();
 
 	// If no keys in curve, return the Default value we passed in.
 	if( NumKeys == 0 )
 	{
-		return 0.f;
+		return DefaultValue;
 	}
 
 	// If only one point, or before the first point in the curve, return the first points value.
diff --git a/Engine/Source/Runtime/MovieSceneCore/Classes/MovieScenePropertyTrack.h b/Engine/Source/Runtime/MovieSceneCore/Classes/MovieScenePropertyTrack.h
--- a/Engine/Source/Runtime/MovieSceneCore/Classes/MovieScenePropertyTrack.h
+++ b/Engine/Source/Runtime/MovieSceneCore/Classes/MovieScenePropertyTrack.h
@@ -5,7 +5,6 @@
 #include "MovieSceneTrack.h"
 #include "MovieScenePropertyTrack.generated.h"
 
-
 /**
  * Base class for tracks that animate an object property
  */
@@ -30,11 +29,14 @@ public:
 	 *
 	 * @param InPropertyName	The property being animated
 	 */
-	virtual void SetPropertyName( FName InPropertyName );
+	virtual void SetPropertyNameAndPath( FName InPropertyName, const FString& InPropertyPath );
 
 	/** @return the name of the property being animataed by this track */
 	FName GetPropertyName() const { return PropertyName; }
 	
+	/** @return The property path for this track */
+	const FString& GetPropertyPath() const { return PropertyPath; }
+
 	/** Sets this track as showable permanently */
 	void SetAsShowable() {bSectionsAreShowable = true;}
 
@@ -53,6 +55,10 @@ protected:
 	UPROPERTY()
 	FName PropertyName;
 
+	/** Path to the property from the source object being changed */
+	UPROPERTY()
+	FString PropertyPath;
+
 	/** All the sections in this list */
 	UPROPERTY()
 	TArray<UMovieSceneSection*> Sections;
diff --git a/Engine/Source/Runtime/MovieSceneCore/Private/MovieSceneCommonHelpers.cpp b/Engine/Source/Runtime/MovieSceneCore/Private/MovieSceneCommonHelpers.cpp
--- a/Engine/Source/Runtime/MovieSceneCore/Private/MovieSceneCommonHelpers.cpp
+++ b/Engine/Source/Runtime/MovieSceneCore/Private/MovieSceneCommonHelpers.cpp
@@ -74,8 +74,11 @@ UMovieSceneSection* MovieSceneHelpers::FindSectionAtTime( const TArray<UMovieSce
 }
 
 
-FTrackInstancePropertyBindings::FTrackInstancePropertyBindings( FName PropertyName )
+FTrackInstancePropertyBindings::FTrackInstancePropertyBindings( FName InPropertyName, const FString& InPropertyPath )
+	: PropertyName( InPropertyName )
+	, PropertyPath( InPropertyPath )
 {
+	
 	static const FString Set(TEXT("Set"));
 
 	const FString FunctionString = Set + PropertyName.ToString();
@@ -83,26 +86,76 @@ FTrackInstancePropertyBindings::FTrackInstancePropertyBindings( FName PropertyNa
 	FunctionName = FName(*FunctionString);
 }
 
-void FTrackInstancePropertyBindings::CallFunction( const TArray<UObject*>& InRuntimeObjects, void* FunctionParams )
+void FTrackInstancePropertyBindings::CallFunction( UObject* InRuntimeObject, void* FunctionParams )
+{
+	FPropertyAndFunction PropAndFunction = RuntimeObjectToFunctionMap.FindRef(InRuntimeObject);
+	if(PropAndFunction.Function)
+	{
+		InRuntimeObject->ProcessEvent(PropAndFunction.Function, FunctionParams);
+	}
+}
+
+
+FTrackInstancePropertyBindings::FPropertyAddress FTrackInstancePropertyBindings::FindPropertyRecursive( UObject* Object, void* BasePointer, UStruct* InStruct, TArray<FString>& InPropertyNames, uint32 Index ) const
 {
-	for( UObject* Object : InRuntimeObjects )
+	UProperty* Property = FindField<UProperty>(InStruct, *InPropertyNames[Index]);
+	
+	FTrackInstancePropertyBindings::FPropertyAddress NewAddress;
+
+	UStructProperty* StructProp = Cast<UStructProperty>( Property );
+	if( StructProp )
 	{
-		UFunction* Function = RuntimeObjectToFunctionMap.FindRef(Object);
-		if(Function)
+		NewAddress.Property = StructProp;
+		NewAddress.Address = BasePointer;
+
+		if( InPropertyNames.IsValidIndex(Index+1) )
 		{
-			Object->ProcessEvent(Function, FunctionParams);
+			void* StructContainer = StructProp->ContainerPtrToValuePtr<void>(BasePointer);
+			return FindPropertyRecursive( Object, StructContainer, StructProp->Struct, InPropertyNames, Index+1 );
 		}
+		else
+		{
+			check( StructProp->GetName() == InPropertyNames[Index] );
+		}
+	}
+	else if( Property )
+	{
+		NewAddress.Property = Property;
+		NewAddress.Address = BasePointer;
+	}
+
+	return NewAddress;
+
+}
+
+FTrackInstancePropertyBindings::FPropertyAddress FTrackInstancePropertyBindings::FindProperty( UObject* InObject, const FString& InPropertyPath ) const
+{
+	TArray<FString> PropertyNames;
+
+	InPropertyPath.ParseIntoArray(&PropertyNames, TEXT("."), true);
+
+	if( PropertyNames.Num() > 0 )
+	{
+		return FindPropertyRecursive( InObject, (void*)InObject, InObject->GetClass(), PropertyNames, 0 );
 	}
+	else
+	{
+		return FTrackInstancePropertyBindings::FPropertyAddress();
+	}
+
 }
 
 void FTrackInstancePropertyBindings::UpdateBindings( const TArray<UObject*>& InRuntimeObjects )
 {
 	for(UObject* Object : InRuntimeObjects)
 	{
-		UFunction* Function = Object->FindFunction(FunctionName);
-		if(Function)
+		FPropertyAndFunction PropAndFunction;
+
+		PropAndFunction.Function = Object->FindFunction(FunctionName);
+		if(PropAndFunction.Function)
 		{
-			RuntimeObjectToFunctionMap.Add(Object, Function);
+			PropAndFunction.PropertyAddress = FindProperty( Object, PropertyPath );
+			RuntimeObjectToFunctionMap.Add(Object, PropAndFunction);
 		}
 		else
 		{
diff --git a/Engine/Source/Runtime/MovieSceneCore/Private/MovieScenePropertyTrack.cpp b/Engine/Source/Runtime/MovieSceneCore/Private/MovieScenePropertyTrack.cpp
--- a/Engine/Source/Runtime/MovieSceneCore/Private/MovieScenePropertyTrack.cpp
+++ b/Engine/Source/Runtime/MovieSceneCore/Private/MovieScenePropertyTrack.cpp
@@ -11,11 +11,12 @@ UMovieScenePropertyTrack::UMovieScenePropertyTrack( const FPostConstructInitiali
 }
 
 
-void UMovieScenePropertyTrack::SetPropertyName( FName InPropertyName )
+void UMovieScenePropertyTrack::SetPropertyNameAndPath( FName InPropertyName, const FString& InPropertyPath )
 {
-	check( InPropertyName != NAME_None );
+	check( InPropertyName != NAME_None && !InPropertyPath.IsEmpty() );
 
 	PropertyName = InPropertyName;
+	PropertyPath = InPropertyPath;
 }
 
 TArray<UMovieSceneSection*> UMovieScenePropertyTrack::GetAllSections() const
diff --git a/Engine/Source/Runtime/MovieSceneCore/Public/MovieSceneCommonHelpers.h b/Engine/Source/Runtime/MovieSceneCore/Public/MovieSceneCommonHelpers.h
--- a/Engine/Source/Runtime/MovieSceneCore/Public/MovieSceneCommonHelpers.h
+++ b/Engine/Source/Runtime/MovieSceneCore/Public/MovieSceneCommonHelpers.h
@@ -28,12 +28,83 @@ public:
 class MOVIESCENECORE_API FTrackInstancePropertyBindings
 {
 public:
-	FTrackInstancePropertyBindings( FName PropertyName );
+	FTrackInstancePropertyBindings( FName InPropertyName, const FString& InPropertyPath );
 
-	void CallFunction( const TArray<UObject*>& InRuntimeObjects, void* FunctionParams );
+	/**
+	 * Calls the setter function for a specific runtime object
+	 *
+	 * @param InRuntimeObject	The runtime object whose function to call
+	 * @param FunctionParams	Parameters to pass to the function
+	 */
+	void CallFunction( UObject* InRuntimeObject, void* FunctionParams );
+
+	/**
+	 * Rebuilds the property and function mappings for a set of runtime objects
+	 *
+	 * @param InRuntimeObjects	The objects to rebuild mappings for
+	 */
 	void UpdateBindings( const TArray<UObject*>& InRuntimeObjects );
+
+	/**
+	 * Gets the current value of a property on an object
+	 *
+	 * @param Object	The object to get the property from
+	 * @return ValueType	The current value
+	 */
+	template <typename ValueType>
+	ValueType GetCurrentValue(const UObject* Object) const
+	{
+		FPropertyAndFunction PropAndFunction = RuntimeObjectToFunctionMap.FindRef(Object);
+
+		if(PropAndFunction.PropertyAddress.Address)
+		{
+			const ValueType* Val = PropAndFunction.PropertyAddress.Property->ContainerPtrToValuePtr<ValueType>(PropAndFunction.PropertyAddress.Address);
+			if(Val)
+			{
+				return *Val;
+			}
+		}
+
+		return ValueType();
+	}
+private:
+
+	struct FPropertyAddress
+	{
+		UProperty* Property;
+		void* Address;
+
+		FPropertyAddress()
+			: Property(nullptr)
+			, Address(nullptr)
+		{}
+	};
+
+	struct FPropertyAndFunction
+	{
+		FPropertyAddress PropertyAddress;
+		UFunction* Function;
+
+		FPropertyAndFunction()
+			: PropertyAddress()
+			, Function( nullptr )
+		{}
+	};
+
+	FPropertyAddress FindPropertyRecursive(UObject* Object, void* BasePointer, UStruct* InStruct, TArray<FString>& InPropertyNames, uint32 Index) const;
+	FPropertyAddress FindProperty(UObject* Object, const FString& InPropertyPath) const;
+
 private:
 	/** Mapping of objects to bound functions that will be called to update data on the track */
-	TMap< TWeakObjectPtr<UObject>, UFunction* > RuntimeObjectToFunctionMap;
+	TMap< TWeakObjectPtr<UObject>, FPropertyAndFunction > RuntimeObjectToFunctionMap;
+
+	/** Path to the property we are bound to */
+	FString PropertyPath;
+
+	/** Name of the function to call to set values */
 	FName FunctionName;
+
+	/** Actual name of the property we are bound to */
+	FName PropertyName;
+
 };
\ No newline at end of file
diff --git a/Engine/Source/Runtime/MovieSceneCoreTypes/Classes/MovieSceneColorSection.h b/Engine/Source/Runtime/MovieSceneCoreTypes/Classes/MovieSceneColorSection.h
--- a/Engine/Source/Runtime/MovieSceneCoreTypes/Classes/MovieSceneColorSection.h
+++ b/Engine/Source/Runtime/MovieSceneCoreTypes/Classes/MovieSceneColorSection.h
@@ -23,7 +23,7 @@ public:
 	 *
 	 * @param Position	The position in time within the movie scene
 	 */
-	virtual FLinearColor Eval( float Position ) const;
+	virtual FLinearColor Eval( float Position, const FLinearColor& DefaultColor ) const;
 
 	/** 
 	 * Adds a key to the section
diff --git a/Engine/Source/Runtime/MovieSceneCoreTypes/Classes/MovieSceneVectorSection.h b/Engine/Source/Runtime/MovieSceneCoreTypes/Classes/MovieSceneVectorSection.h
--- a/Engine/Source/Runtime/MovieSceneCoreTypes/Classes/MovieSceneVectorSection.h
+++ b/Engine/Source/Runtime/MovieSceneCoreTypes/Classes/MovieSceneVectorSection.h
@@ -20,7 +20,7 @@ public:
 	 *
 	 * @param Position	The position in time within the movie scene
 	 */
-	FVector4 Eval( float Position ) const;
+	FVector4 Eval( float Position, const FVector4& DefaultVector ) const;
 
 	/** MovieSceneSection interface */
 	virtual void MoveSection( float DeltaPosition ) override;
diff --git a/Engine/Source/Runtime/MovieSceneCoreTypes/Classes/MovieSceneVectorTrack.h b/Engine/Source/Runtime/MovieSceneCoreTypes/Classes/MovieSceneVectorTrack.h
--- a/Engine/Source/Runtime/MovieSceneCoreTypes/Classes/MovieSceneVectorTrack.h
+++ b/Engine/Source/Runtime/MovieSceneCoreTypes/Classes/MovieSceneVectorTrack.h
@@ -52,7 +52,7 @@ public:
 	 * @param OutVector 	The value at the playback position
 	 * @return true if anything was evaluated. Note: if false is returned OutVector remains unchanged
 	 */
-	virtual bool Eval( float Position, float LastPostion, FVector4& OutVector ) const;
+	virtual bool Eval( float Position, float LastPostion, FVector4& InOutVector ) const;
 
 	/** @return Get the number of channels used by the vector */
 	int32 GetNumChannelsUsed() const { return NumChannelsUsed; }
diff --git a/Engine/Source/Runtime/MovieSceneCoreTypes/Private/Instances/MovieSceneBoolTrackInstance.cpp b/Engine/Source/Runtime/MovieSceneCoreTypes/Private/Instances/MovieSceneBoolTrackInstance.cpp
--- a/Engine/Source/Runtime/MovieSceneCoreTypes/Private/Instances/MovieSceneBoolTrackInstance.cpp
+++ b/Engine/Source/Runtime/MovieSceneCoreTypes/Private/Instances/MovieSceneBoolTrackInstance.cpp
@@ -8,7 +8,7 @@ FMovieSceneBoolTrackInstance::FMovieSceneBoolTrackInstance( UMovieSceneBoolTrack
 {
 	BoolTrack = &InBoolTrack;
 	
-	PropertyBindings = MakeShareable( new FTrackInstancePropertyBindings( BoolTrack->GetPropertyName() ) );
+	PropertyBindings = MakeShareable( new FTrackInstancePropertyBindings( BoolTrack->GetPropertyName(), BoolTrack->GetPropertyPath() ) );
 }
 
 void FMovieSceneBoolTrackInstance::Update( float Position, float LastPosition, const TArray<UObject*>& RuntimeObjects, class IMovieScenePlayer& Player ) 
@@ -16,7 +16,10 @@ void FMovieSceneBoolTrackInstance::Update( float Position, float LastPosition, c
 	bool BoolValue = false;
 	if( BoolTrack->Eval( Position, LastPosition, BoolValue ) )
 	{
-		PropertyBindings->CallFunction( RuntimeObjects, &BoolValue );
+		for( UObject* Object : RuntimeObjects )
+		{
+			PropertyBindings->CallFunction( Object, &BoolValue );
+		}
 	}
 }
 
diff --git a/Engine/Source/Runtime/MovieSceneCoreTypes/Private/Instances/MovieSceneColorTrackInstance.cpp b/Engine/Source/Runtime/MovieSceneCoreTypes/Private/Instances/MovieSceneColorTrackInstance.cpp
--- a/Engine/Source/Runtime/MovieSceneCoreTypes/Private/Instances/MovieSceneColorTrackInstance.cpp
+++ b/Engine/Source/Runtime/MovieSceneCoreTypes/Private/Instances/MovieSceneColorTrackInstance.cpp
@@ -9,15 +9,18 @@ FMovieSceneColorTrackInstance::FMovieSceneColorTrackInstance( UMovieSceneColorTr
 {
 	ColorTrack = &InColorTrack;
 
-	PropertyBindings = MakeShareable( new FTrackInstancePropertyBindings( ColorTrack->GetPropertyName() ) );
+	PropertyBindings = MakeShareable( new FTrackInstancePropertyBindings( ColorTrack->GetPropertyName(), ColorTrack->GetPropertyPath() ) );
 }
 
 void FMovieSceneColorTrackInstance::Update( float Position, float LastPosition, const TArray<UObject*>& RuntimeObjects, class IMovieScenePlayer& Player ) 
 {
-	FLinearColor ColorValue;
-	if( ColorTrack->Eval( Position, LastPosition, ColorValue ) )
+	for(UObject* Object : RuntimeObjects)
 	{
-		PropertyBindings->CallFunction( RuntimeObjects, &ColorValue );
+		FLinearColor ColorValue = PropertyBindings->GetCurrentValue<FLinearColor>(Object);
+		if(ColorTrack->Eval(Position, LastPosition, ColorValue))
+		{
+			PropertyBindings->CallFunction(Object, &ColorValue);
+		}
 	}
 }
 
diff --git a/Engine/Source/Runtime/MovieSceneCoreTypes/Private/Instances/MovieSceneFloatTrackInstance.cpp b/Engine/Source/Runtime/MovieSceneCoreTypes/Private/Instances/MovieSceneFloatTrackInstance.cpp
--- a/Engine/Source/Runtime/MovieSceneCoreTypes/Private/Instances/MovieSceneFloatTrackInstance.cpp
+++ b/Engine/Source/Runtime/MovieSceneCoreTypes/Private/Instances/MovieSceneFloatTrackInstance.cpp
@@ -8,7 +8,7 @@ FMovieSceneFloatTrackInstance::FMovieSceneFloatTrackInstance( UMovieSceneFloatTr
 {
 	FloatTrack = &InFloatTrack;
 
-	PropertyBindings = MakeShareable( new FTrackInstancePropertyBindings( FloatTrack->GetPropertyName() ) );
+	PropertyBindings = MakeShareable( new FTrackInstancePropertyBindings( FloatTrack->GetPropertyName(), FloatTrack->GetPropertyPath() ) );
 }
 
 void FMovieSceneFloatTrackInstance::Update( float Position, float LastPosition, const TArray<UObject*>& RuntimeObjects, class IMovieScenePlayer& Player ) 
@@ -16,7 +16,10 @@ void FMovieSceneFloatTrackInstance::Update( float Position, float LastPosition,
 	float FloatValue = 0.0f;
 	if( FloatTrack->Eval( Position, LastPosition, FloatValue ) )
 	{
-		PropertyBindings->CallFunction( RuntimeObjects, &FloatValue );
+		for(UObject* Object : RuntimeObjects)
+		{
+			PropertyBindings->CallFunction( Object, &FloatValue );
+		}
 	}
 }
 
diff --git a/Engine/Source/Runtime/MovieSceneCoreTypes/Private/Instances/MovieSceneVectorTrackInstance.cpp b/Engine/Source/Runtime/MovieSceneCoreTypes/Private/Instances/MovieSceneVectorTrackInstance.cpp
--- a/Engine/Source/Runtime/MovieSceneCoreTypes/Private/Instances/MovieSceneVectorTrackInstance.cpp
+++ b/Engine/Source/Runtime/MovieSceneCoreTypes/Private/Instances/MovieSceneVectorTrackInstance.cpp
@@ -7,7 +7,7 @@ FMovieSceneVectorTrackInstance::FMovieSceneVectorTrackInstance( UMovieSceneVecto
 {
 	VectorTrack = &InVectorTrack;
 
-	PropertyBindings = MakeShareable( new FTrackInstancePropertyBindings( VectorTrack->GetPropertyName() ) );
+	PropertyBindings = MakeShareable( new FTrackInstancePropertyBindings( VectorTrack->GetPropertyName(), VectorTrack->GetPropertyPath() ) );
 }
 
 void FMovieSceneVectorTrackInstance::Update( float Position, float LastPosition, const TArray<UObject*>& RuntimeObjects, class IMovieScenePlayer& Player ) 
@@ -21,18 +21,27 @@ void FMovieSceneVectorTrackInstance::Update( float Position, float LastPosition,
 			case 2:
 			{
 				FVector2D Value(Vector.X, Vector.Y);
-				PropertyBindings->CallFunction(RuntimeObjects, &Value);
+				for(UObject* Object : RuntimeObjects)
+				{
+					PropertyBindings->CallFunction(Object, &Value);
+				}
 				break;
 			}
 			case 3:
 			{
 				FVector Value(Vector.X, Vector.Y, Vector.Z);
-				PropertyBindings->CallFunction(RuntimeObjects, &Value);
+				for(UObject* Object : RuntimeObjects)
+				{
+					PropertyBindings->CallFunction(Object, &Value);
+				}
 				break;
 			}
 			case 4:
 			{
-				PropertyBindings->CallFunction(RuntimeObjects, &Vector);
+				for(UObject* Object : RuntimeObjects)
+				{
+					PropertyBindings->CallFunction(Object, &Vector);
+				}
 				break;
 			}
 			default:
diff --git a/Engine/Source/Runtime/MovieSceneCoreTypes/Private/MovieSceneColorSection.cpp b/Engine/Source/Runtime/MovieSceneCoreTypes/Private/MovieSceneColorSection.cpp
--- a/Engine/Source/Runtime/MovieSceneCoreTypes/Private/MovieSceneColorSection.cpp
+++ b/Engine/Source/Runtime/MovieSceneCoreTypes/Private/MovieSceneColorSection.cpp
@@ -31,12 +31,12 @@ void UMovieSceneColorSection::DilateSection( float DilationFactor, float Origin
 
 }
 
-FLinearColor UMovieSceneColorSection::Eval( float Position ) const
+FLinearColor UMovieSceneColorSection::Eval( float Position, const FLinearColor& DefaultColor ) const
 {
-	return FLinearColor(RedCurve.Eval(Position),
-						GreenCurve.Eval(Position),
-						BlueCurve.Eval(Position),
-						AlphaCurve.Eval(Position));
+	return FLinearColor(RedCurve.Eval(Position, DefaultColor.R),
+						GreenCurve.Eval(Position, DefaultColor.G),
+						BlueCurve.Eval(Position, DefaultColor.B),
+						AlphaCurve.Eval(Position, DefaultColor.A));
 }
 
 void UMovieSceneColorSection::AddKey( float Time, const FColorKey& Key )
@@ -62,7 +62,7 @@ bool UMovieSceneColorSection::NewKeyIsNewData(float Time, FLinearColor Value) co
 		GreenCurve.GetNumKeys() == 0 ||
 		BlueCurve.GetNumKeys() == 0 ||
 		AlphaCurve.GetNumKeys() == 0 ||
-		!Eval(Time).Equals(Value);
+		!Eval(Time,Value).Equals(Value);
 }
 
 void UMovieSceneColorSection::AddKeyToNamedCurve(float Time, const FColorKey& Key)
diff --git a/Engine/Source/Runtime/MovieSceneCoreTypes/Private/MovieSceneColorTrack.cpp b/Engine/Source/Runtime/MovieSceneCoreTypes/Private/MovieSceneColorTrack.cpp
--- a/Engine/Source/Runtime/MovieSceneCoreTypes/Private/MovieSceneColorTrack.cpp
+++ b/Engine/Source/Runtime/MovieSceneCoreTypes/Private/MovieSceneColorTrack.cpp
@@ -44,7 +44,7 @@ bool UMovieSceneColorTrack::Eval( float Position, float LastPosition, FLinearCol
 
 	if( Section )
 	{
-		OutColor = CastChecked<UMovieSceneColorSection>( Section )->Eval( Position );
+		OutColor = CastChecked<UMovieSceneColorSection>( Section )->Eval( Position, OutColor );
 	}
 
 	return Section != NULL;
diff --git a/Engine/Source/Runtime/MovieSceneCoreTypes/Private/MovieSceneVectorSection.cpp b/Engine/Source/Runtime/MovieSceneCoreTypes/Private/MovieSceneVectorSection.cpp
--- a/Engine/Source/Runtime/MovieSceneCoreTypes/Private/MovieSceneVectorSection.cpp
+++ b/Engine/Source/Runtime/MovieSceneCoreTypes/Private/MovieSceneVectorSection.cpp
@@ -10,13 +10,13 @@ UMovieSceneVectorSection::UMovieSceneVectorSection( const FPostConstructInitiali
 
 }
 
-FVector4 UMovieSceneVectorSection::Eval( float Position ) const
+FVector4 UMovieSceneVectorSection::Eval( float Position, const FVector4& DefaultVector ) const
 {
 	return FVector4(
-		Curves[0].Eval( Position ),
-		Curves[1].Eval( Position ),
-		Curves[2].Eval( Position ),
-		Curves[3].Eval( Position ) );
+		Curves[0].Eval( Position, DefaultVector.X ),
+		Curves[1].Eval( Position, DefaultVector.Y ),
+		Curves[2].Eval( Position, DefaultVector.Z ),
+		Curves[3].Eval( Position, DefaultVector.W ) );
 }
 
 void UMovieSceneVectorSection::AddKey( float Time, FName CurveName, const FVector4& Value )
diff --git a/Engine/Source/Runtime/MovieSceneCoreTypes/Private/MovieSceneVectorTrack.cpp b/Engine/Source/Runtime/MovieSceneCoreTypes/Private/MovieSceneVectorTrack.cpp
--- a/Engine/Source/Runtime/MovieSceneCoreTypes/Private/MovieSceneVectorTrack.cpp
+++ b/Engine/Source/Runtime/MovieSceneCoreTypes/Private/MovieSceneVectorTrack.cpp
@@ -58,13 +58,13 @@ bool UMovieSceneVectorTrack::AddKeyToSection( float Time, const FVectorKey<FVect
 	return AddKeyToSection(Time, FVector4(Key.Value.X, Key.Value.Y, 0.f, 0.f), 2, Key.CurveName, Key.bAddKeyEvenIfUnchanged );
 }
 
-bool UMovieSceneVectorTrack::Eval( float Position, float LastPosition, FVector4& OutVector ) const
+bool UMovieSceneVectorTrack::Eval( float Position, float LastPosition, FVector4& InOutVector ) const
 {
 	const UMovieSceneSection* Section = MovieSceneHelpers::FindSectionAtTime( Sections, Position );
 
 	if( Section )
 	{
-		OutVector = CastChecked<UMovieSceneVectorSection>( Section )->Eval( Position );
+		InOutVector = CastChecked<UMovieSceneVectorSection>( Section )->Eval( Position, InOutVector );
 	}
 
 	return Section != NULL;
diff --git a/Engine/Source/Runtime/UMG/Private/Animation/MovieSceneMarginTrackInstance.cpp b/Engine/Source/Runtime/UMG/Private/Animation/MovieSceneMarginTrackInstance.cpp
--- a/Engine/Source/Runtime/UMG/Private/Animation/MovieSceneMarginTrackInstance.cpp
+++ b/Engine/Source/Runtime/UMG/Private/Animation/MovieSceneMarginTrackInstance.cpp
@@ -9,15 +9,18 @@ FMovieSceneMarginTrackInstance::FMovieSceneMarginTrackInstance( UMovieSceneMargi
 {
 	MarginTrack = &InMarginTrack;
 
-	PropertyBindings = MakeShareable( new FTrackInstancePropertyBindings( MarginTrack->GetPropertyName() ) );
+	PropertyBindings = MakeShareable( new FTrackInstancePropertyBindings( MarginTrack->GetPropertyName(), MarginTrack->GetPropertyPath() ) );
 }
 
 void FMovieSceneMarginTrackInstance::Update( float Position, float LastPosition, const TArray<UObject*>& RuntimeObjects, class IMovieScenePlayer& Player ) 
 {
-	FMargin MarginValue;
-	if( MarginTrack->Eval( Position, LastPosition, MarginValue ) )
+	for( UObject* Object : RuntimeObjects )
 	{
-		PropertyBindings->CallFunction( RuntimeObjects, &MarginValue );
+		FMargin MarginValue = PropertyBindings->GetCurrentValue<FMargin>( Object );
+		if(MarginTrack->Eval(Position, LastPosition, MarginValue))
+		{
+			PropertyBindings->CallFunction(Object, &MarginValue);
+		}
 	}
 }
 
diff --git a/Engine/Source/Runtime/UMG/Public/Animation/MovieSceneMarginSection.h b/Engine/Source/Runtime/UMG/Public/Animation/MovieSceneMarginSection.h
--- a/Engine/Source/Runtime/UMG/Public/Animation/MovieSceneMarginSection.h
+++ b/Engine/Source/Runtime/UMG/Public/Animation/MovieSceneMarginSection.h
@@ -25,7 +25,7 @@ public:
 	 *
 	 * @param Position	The position in time within the movie scene
 	 */
-	virtual FMargin Eval( float Position ) const;
+	FMargin Eval( float Position, const FMargin& DefaultValue ) const;
 
 	/** 
 	 * Adds a key to the section
@@ -74,7 +74,6 @@ public:
 	 */
 	FRichCurve& GetBottomCurve() { return BottomCurve; }
 	const FRichCurve& GetBottomCurve() const { return BottomCurve; }
-
 private:
 	void AddKeyToNamedCurve( float Time, const FMarginKey& MarginKey );
 	/**
