- Fixed clang compile error

diff --git a/Engine/Source/Runtime/MovieSceneCore/Private/MovieSceneCommonHelpers.cpp b/Engine/Source/Runtime/MovieSceneCore/Private/MovieSceneCommonHelpers.cpp
--- a/Engine/Source/Runtime/MovieSceneCore/Private/MovieSceneCommonHelpers.cpp
+++ b/Engine/Source/Runtime/MovieSceneCore/Private/MovieSceneCommonHelpers.cpp
@@ -75,8 +75,8 @@ UMovieSceneSection* MovieSceneHelpers::FindSectionAtTime( const TArray<UMovieSce
 
 
 FTrackInstancePropertyBindings::FTrackInstancePropertyBindings( FName InPropertyName, const FString& InPropertyPath )
-	: PropertyName( InPropertyName )
-	, PropertyPath( InPropertyPath )
+    : PropertyPath( InPropertyPath )
+	, PropertyName( InPropertyName )
 {
 	
 	static const FString Set(TEXT("Set"));
