Fixed sequencer data structures that were persisting FTexts instead of FStrings as UProperties, causing them to be gathered for localization. Ultimately they are turned into FText, but that is done at the appropriate time.

diff --git a/Engine/Source/Developer/MovieSceneTools/Private/K2Node_PlayMovieScene.cpp b/Engine/Source/Developer/MovieSceneTools/Private/K2Node_PlayMovieScene.cpp
--- a/Engine/Source/Developer/MovieSceneTools/Private/K2Node_PlayMovieScene.cpp
+++ b/Engine/Source/Developer/MovieSceneTools/Private/K2Node_PlayMovieScene.cpp
@@ -284,7 +284,7 @@ void UK2Node_PlayMovieScene::CreatePinForBoundObject( FMovieSceneBoundObject& Bo
 				if( Possessable.GetGuid() == BoundObject.GetPossessableGuid() )
 				{
 					// Found a name for this possessable
-					PinFriendlyName = Possessable.GetDisplayName();
+					PinFriendlyName = FText::FromString(Possessable.GetName());
 					break;
 				}
 			}
diff --git a/Engine/Source/Editor/Sequencer/Private/Sequencer.cpp b/Engine/Source/Editor/Sequencer/Private/Sequencer.cpp
--- a/Engine/Source/Editor/Sequencer/Private/Sequencer.cpp
+++ b/Engine/Source/Editor/Sequencer/Private/Sequencer.cpp
@@ -450,8 +450,7 @@ void FSequencer::OnActorsDropped( const TArray<TWeakObjectPtr<AActor> >& Actors
 				// Create a new possessable
 				OwnerMovieScene->Modify();
 				
-				const FText PossessableName = FText::FromString( Actor->GetActorLabel() );
-				const FGuid PossessableGuid = OwnerMovieScene->AddPossessable( PossessableName, Actor->GetClass() );
+				const FGuid PossessableGuid = OwnerMovieScene->AddPossessable( Actor->GetActorLabel(), Actor->GetClass() );
 				 
 				if ( IsShotFilteringOn() )
 				{
@@ -549,8 +548,7 @@ FGuid FSequencer::GetHandleToObject( UObject* Object )
 			// Create a new possessable
 			OwnerMovieScene->Modify();
 			
-			const FText PossessableName = FText::FromString( Object->GetName() );
-			ObjectGuid = OwnerMovieScene->AddPossessable( PossessableName, Object->GetClass() );
+			ObjectGuid = OwnerMovieScene->AddPossessable( Object->GetName(), Object->GetClass() );
 			
 			if ( IsShotFilteringOn() )
 			{
@@ -911,7 +909,7 @@ FGuid FSequencer::AddSpawnableForAssetOrClass( UObject* Object, UObject* Counter
 		const FName BlueprintName = MakeUniqueObjectName( OwnerMovieScene, UBlueprint::StaticClass(), AssetName );
 
 		// Use the asset name as the initial spawnable name
-		const FText NewSpawnableName = FText::FromName( AssetName );		// @todo sequencer: Need UI to allow user to rename these slots
+		const FString NewSpawnableName = AssetName.ToString();		// @todo sequencer: Need UI to allow user to rename these slots
 
 		// Create our new blueprint!
 		UBlueprint* NewBlueprint = NULL;
diff --git a/Engine/Source/Editor/Sequencer/Private/SequencerActorBindingManager.cpp b/Engine/Source/Editor/Sequencer/Private/SequencerActorBindingManager.cpp
--- a/Engine/Source/Editor/Sequencer/Private/SequencerActorBindingManager.cpp
+++ b/Engine/Source/Editor/Sequencer/Private/SequencerActorBindingManager.cpp
@@ -216,7 +216,7 @@ void FSequencerActorBindingManager::SpawnOrDestroyObjectsForInstance( TSharedRef
 						if( NewActor )
 						{
 							// @todo sequencer: We're naming the actor based off of the spawnable's name.  Is that really what we want?
-							GEditor->SetActorLabelUnique( NewActor, Spawnable.GetDisplayName().ToString() );
+							GEditor->SetActorLabelUnique( NewActor, Spawnable.GetName() );
 
 							// Actor was spawned OK!
 
diff --git a/Engine/Source/Editor/Sequencer/Private/SequencerDisplayNode.h b/Engine/Source/Editor/Sequencer/Private/SequencerDisplayNode.h
--- a/Engine/Source/Editor/Sequencer/Private/SequencerDisplayNode.h
+++ b/Engine/Source/Editor/Sequencer/Private/SequencerDisplayNode.h
@@ -368,7 +368,7 @@ public:
 	 * @param InParentNode		The parent of this node or NULL if this is a root node
 	 * @param InParentTree		The tree this node is in
 	 */
-	FObjectBindingNode( FName NodeName, const FText& InObjectName, const FGuid& InObjectBinding, TSharedPtr<FSequencerDisplayNode> InParentNode, FSequencerNodeTree& InParentTree )
+	FObjectBindingNode( FName NodeName, const FString& InObjectName, const FGuid& InObjectBinding, TSharedPtr<FSequencerDisplayNode> InParentNode, FSequencerNodeTree& InParentTree )
 		: FSequencerDisplayNode( NodeName, InParentNode, InParentTree )
 		, ObjectBinding( InObjectBinding )
 		, DisplayName( InObjectName )
@@ -376,7 +376,7 @@ public:
 
 	/** FSequencerDisplayNodeInterface */
 	virtual ESequencerNode::Type GetType() const override { return ESequencerNode::Object; }
-	virtual FText GetDisplayName() const override { return DisplayName; }
+	virtual FText GetDisplayName() const override { return FText::FromString(DisplayName); }
 	virtual float GetNodeHeight() const override;
 	virtual bool GetShotFilteredVisibilityToCache() const override;
 	
@@ -390,7 +390,7 @@ private:
 	/** The binding to live objects */
 	FGuid ObjectBinding;
 	/** Display name of the object */
-	FText DisplayName;
+	FString DisplayName;
 };
 
 /**
diff --git a/Engine/Source/Editor/Sequencer/Private/SequencerNodeTree.cpp b/Engine/Source/Editor/Sequencer/Private/SequencerNodeTree.cpp
--- a/Engine/Source/Editor/Sequencer/Private/SequencerNodeTree.cpp
+++ b/Engine/Source/Editor/Sequencer/Private/SequencerNodeTree.cpp
@@ -48,7 +48,7 @@ void FSequencerNodeTree::Update()
 	// Make nodes for all object bindings
 	for( int32 BindingIndex = 0; BindingIndex < ObjectBindings.Num(); ++BindingIndex )
 	{
-		TSharedRef<FObjectBindingNode> ObjectBindingNode = AddObjectBinding( ObjectBindings[BindingIndex].GetDisplayName(), ObjectBindings[BindingIndex].GetObjectGuid() );
+		TSharedRef<FObjectBindingNode> ObjectBindingNode = AddObjectBinding( ObjectBindings[BindingIndex].GetName(), ObjectBindings[BindingIndex].GetObjectGuid() );
 
 		const TArray<UMovieSceneTrack*>& Tracks = ObjectBindings[BindingIndex].GetTracks();
 
@@ -133,7 +133,7 @@ const TArray< TSharedRef<FSequencerDisplayNode> >& FSequencerNodeTree::GetRootNo
 }
 
 
-TSharedRef<FObjectBindingNode> FSequencerNodeTree::AddObjectBinding( const FText& ObjectName, const FGuid& ObjectBinding )
+TSharedRef<FObjectBindingNode> FSequencerNodeTree::AddObjectBinding( const FString& ObjectName, const FGuid& ObjectBinding )
 {
 	// The node name is the object guid
 	FName ObjectNodeName = *ObjectBinding.ToString();
diff --git a/Engine/Source/Editor/Sequencer/Private/SequencerNodeTree.h b/Engine/Source/Editor/Sequencer/Private/SequencerNodeTree.h
--- a/Engine/Source/Editor/Sequencer/Private/SequencerNodeTree.h
+++ b/Engine/Source/Editor/Sequencer/Private/SequencerNodeTree.h
@@ -120,7 +120,7 @@ private:
 	 * 
 	 * @param ObjectBinding	The object binding guid
 	 */
-	TSharedRef<FObjectBindingNode> AddObjectBinding( const FText& ObjectName, const FGuid& ObjectBinding );
+	TSharedRef<FObjectBindingNode> AddObjectBinding( const FString& ObjectName, const FGuid& ObjectBinding );
 private:
 	/** Tools for building movie scene section layouts.  One tool for each track */
 	TMap< class UMovieSceneTrack*, TSharedPtr<FMovieSceneTrackEditor> > EditorMap;
diff --git a/Engine/Source/Runtime/MovieSceneCore/Classes/MovieScene.h b/Engine/Source/Runtime/MovieSceneCore/Classes/MovieScene.h
--- a/Engine/Source/Runtime/MovieSceneCore/Classes/MovieScene.h
+++ b/Engine/Source/Runtime/MovieSceneCore/Classes/MovieScene.h
@@ -22,7 +22,7 @@ public:
 	}
 
 	/** FMovieSceneSpawnable initialization constructor */
-	FMovieSceneSpawnable( const FText& InitName, UClass* InitClass, UObject* InitCounterpartGamePreviewObject );
+	FMovieSceneSpawnable( const FString& InitName, UClass* InitClass, UObject* InitCounterpartGamePreviewObject );
 
 	/** @return Returns the guid for this possessable */
 	const FGuid& GetGuid() const
@@ -31,7 +31,7 @@ public:
 	}
 
 	/** @return Returns the name of this spawnable */
-	const FText& GetDisplayName() const
+	const FString& GetName() const
 	{
 		return Name;
 	}
@@ -56,7 +56,7 @@ private:
 	/** Name label */
 	// @todo sequencer: Should be editor-only probably
 	UPROPERTY()
-	FText Name;
+	FString Name;
 
 	/** Data-only blueprint-generated-class for this object */
 	// @todo sequencer: Could be weak object ptr here, IF blueprints that are inners are housekept properly without references
@@ -85,7 +85,7 @@ public:
 	}
 
 	/** FMovieScenePossessable initialization constructor */
-	FMovieScenePossessable( const FText& InitName, UClass* InitPossessedObjectClass );
+	FMovieScenePossessable( const FString& InitName, UClass* InitPossessedObjectClass );
 
 	/** @return Returns the guid for this possessable */
 	const FGuid& GetGuid() const
@@ -94,7 +94,7 @@ public:
 	}
 
 	/** @return Returns the name of this spawnable */
-	const FText& GetDisplayName() const
+	const FString& GetName() const
 	{
 		return Name;
 	}
@@ -116,7 +116,7 @@ private:
 	/** Name label for this slot */
 	// @todo sequencer: Should be editor-only probably
 	UPROPERTY()
-	FText Name;
+	FString Name;
 
 	/** Type of the object we'll be possessing */
 	// @todo sequencer: Might be able to be editor-only.  We'll see.
@@ -152,13 +152,13 @@ struct FMovieSceneObjectBinding
 	FMovieSceneObjectBinding()
 	{}
 
-	FMovieSceneObjectBinding( const FGuid& InObjectGuid, const FText& InBindingName, const TArray<class UMovieSceneTrack*>& InTracks )
+	FMovieSceneObjectBinding( const FGuid& InObjectGuid, const FString& InBindingName, const TArray<class UMovieSceneTrack*>& InTracks )
 		: ObjectGuid( InObjectGuid )
 		, BindingName( InBindingName )
 		, Tracks( InTracks )
 	{}
 
-	FMovieSceneObjectBinding( const FGuid& InObjectGuid, const FText& InBindingName )
+	FMovieSceneObjectBinding( const FGuid& InObjectGuid, const FString& InBindingName )
 		: ObjectGuid( InObjectGuid )
 		, BindingName( InBindingName )
 	{}
@@ -176,7 +176,7 @@ struct FMovieSceneObjectBinding
 	/**
 	 * @return The display name of the binding
 	 */
-	const FText& GetDisplayName() const { return BindingName; } 
+	const FString& GetName() const { return BindingName; } 
 
 	/**
 	 * Adds a new track to this binding
@@ -204,7 +204,7 @@ private:
 	
 	/** Display name */
 	UPROPERTY()
-	FText BindingName;
+	FString BindingName;
 
 	/** All tracks in this binding */
 	UPROPERTY()
@@ -232,7 +232,7 @@ public:
 	 *
 	 * @return	Guid of the newly-added spawnable
 	 */
-	virtual FGuid AddSpawnable( const FText& Name, class UBlueprint* Blueprint, UObject* CounterpartGamePreviewObject );
+	virtual FGuid AddSpawnable( const FString& Name, class UBlueprint* Blueprint, UObject* CounterpartGamePreviewObject );
 
 	/**
 	 * Removes a spawnable from this movie scene.
@@ -283,7 +283,7 @@ public:
 	 *
 	 * @return	Guid of the newly-added possessable
 	 */
-	virtual FGuid AddPossessable( const FText& Name, UClass* Class );
+	virtual FGuid AddPossessable( const FString& Name, UClass* Class );
 
 	/**
 	 * Removes a possessable from this movie scene.
diff --git a/Engine/Source/Runtime/MovieSceneCore/Private/MovieScene.cpp b/Engine/Source/Runtime/MovieSceneCore/Private/MovieScene.cpp
--- a/Engine/Source/Runtime/MovieSceneCore/Private/MovieScene.cpp
+++ b/Engine/Source/Runtime/MovieSceneCore/Private/MovieScene.cpp
@@ -4,7 +4,7 @@
 #include "MovieScene.h"
 
 
-FMovieSceneSpawnable::FMovieSceneSpawnable( const FText& InitName, UClass* InitClass, UObject* InitCounterpartGamePreviewObject )
+FMovieSceneSpawnable::FMovieSceneSpawnable( const FString& InitName, UClass* InitClass, UObject* InitCounterpartGamePreviewObject )
 {
 	Guid = FGuid::NewGuid();
 	Name = InitName;
@@ -14,7 +14,7 @@ FMovieSceneSpawnable::FMovieSceneSpawnable( const FText& InitName, UClass* InitC
 
 
 
-FMovieScenePossessable::FMovieScenePossessable( const FText& InitName, UClass* InitPossessedObjectClass )
+FMovieScenePossessable::FMovieScenePossessable( const FString& InitName, UClass* InitPossessedObjectClass )
 {
 	Guid = FGuid::NewGuid();
 	Name = InitName;
@@ -49,7 +49,7 @@ UMovieScene::UMovieScene( const FObjectInitializer& ObjectInitializer )
 
 #if WITH_EDITOR
 // @todo sequencer: Some of these methods should only be used by tools, and should probably move out of MovieSceneCore!
-FGuid UMovieScene::AddSpawnable( const FText& Name, UBlueprint* Blueprint, UObject* CounterpartGamePreviewObject )
+FGuid UMovieScene::AddSpawnable( const FString& Name, UBlueprint* Blueprint, UObject* CounterpartGamePreviewObject )
 {
 	check( (Blueprint != NULL) && (Blueprint->GeneratedClass) );
 
@@ -59,7 +59,7 @@ FGuid UMovieScene::AddSpawnable( const FText& Name, UBlueprint* Blueprint, UObje
 	Spawnables.Add( NewSpawnable );
 
 	// Add a new binding so that tracks can be added to it
-	new (ObjectBindings) FMovieSceneObjectBinding( NewSpawnable.GetGuid(), NewSpawnable.GetDisplayName() );
+	new (ObjectBindings) FMovieSceneObjectBinding( NewSpawnable.GetGuid(), NewSpawnable.GetName() );
 
 	return NewSpawnable.GetGuid();
 }
@@ -150,7 +150,7 @@ const FMovieSceneSpawnable* UMovieScene::FindSpawnableForCounterpart( UObject* G
 }
 
 
-FGuid UMovieScene::AddPossessable( const FText& Name, UClass* Class )
+FGuid UMovieScene::AddPossessable( const FString& Name, UClass* Class )
 {
 	Modify();
 
@@ -158,7 +158,7 @@ FGuid UMovieScene::AddPossessable( const FText& Name, UClass* Class )
 	Possessables.Add( NewPossessable );
 
 	// Add a new binding so that tracks can be added to it
-	new (ObjectBindings) FMovieSceneObjectBinding( NewPossessable.GetGuid(), NewPossessable.GetDisplayName() );
+	new (ObjectBindings) FMovieSceneObjectBinding( NewPossessable.GetGuid(), NewPossessable.GetName() );
 
 	return NewPossessable.GetGuid();
 }
