Sequencer: Split out finding nearest section in FindSectionAtTime into a separate function, FindNearestSectionAtTime. Eval() continues to call FindSectionAtTime(), but it now doesn't  evaluate keys outside of sections. AddKey() now calls FindNearestSectionAtTime() so that it can still add keys when the current time doesn't overlap a section (the nearest section just gets extended to include the current time).

diff --git a/Engine/Source/Runtime/MovieSceneCore/Private/MovieSceneCommonHelpers.cpp b/Engine/Source/Runtime/MovieSceneCore/Private/MovieSceneCommonHelpers.cpp
--- a/Engine/Source/Runtime/MovieSceneCore/Private/MovieSceneCommonHelpers.cpp
+++ b/Engine/Source/Runtime/MovieSceneCore/Private/MovieSceneCommonHelpers.cpp
@@ -29,6 +29,22 @@ TArray<UMovieSceneSection*> MovieSceneHelpers::GetTraversedSections( const TArra
 }
 
 UMovieSceneSection* MovieSceneHelpers::FindSectionAtTime( const TArray<UMovieSceneSection*>& Sections, float Time )
+{
+	for( int32 SectionIndex = 0; SectionIndex < Sections.Num(); ++SectionIndex )
+	{
+		UMovieSceneSection* Section = Sections[SectionIndex];
+
+		//@todo Sequencer - There can be multiple sections overlapping in time. Returning instantly does not account for that.
+		if( Section->IsTimeWithinSection( Time ) )
+		{
+			return Section;
+		}
+	}
+
+	return NULL;
+}
+
+UMovieSceneSection* MovieSceneHelpers::FindNearestSectionAtTime( const TArray<UMovieSceneSection*>& Sections, float Time )
 {
 	// Only update the section if the position is within the time span of the section
 	// Or if there are no sections at the time, the left closest section to the time
@@ -73,17 +89,38 @@ UMovieSceneSection* MovieSceneHelpers::FindSectionAtTime( const TArray<UMovieSce
 	return ClosestSection ? ClosestSection : EarliestSection;
 }
 
+void MovieSceneHelpers::GatherAllKeyIndicesInSection(const UMovieSceneSection* Section, const FIndexedCurve& Curve, TSet<int32>& KeyIndices)
+{
+	KeyIndices.Empty();
+
+	float StartTime = Section->GetStartTime();
+	float EndTime = Section->GetEndTime();
+
+	///@todo this probably isn't efficient?
+	for (auto It(Curve.GetKeyHandleIterator()); It; ++It)
+	{
+		float Time = Curve.GetKeyTime(It.Key());
+		if (Time >= StartTime && Time <= EndTime)
+			KeyIndices.Add(It.Value());
+	}
+}
 
-FTrackInstancePropertyBindings::FTrackInstancePropertyBindings( FName InPropertyName, const FString& InPropertyPath )
+FTrackInstancePropertyBindings::FTrackInstancePropertyBindings( FName InPropertyName, const FString& InPropertyPath, const FName& InFunctionName )
     : PropertyPath( InPropertyPath )
 	, PropertyName( InPropertyName )
 {
-	
-	static const FString Set(TEXT("Set"));
+	if (InFunctionName != FName())
+	{
+		FunctionName = InFunctionName;
+	}
+	else
+	{
+		static const FString Set(TEXT("Set"));
 
-	const FString FunctionString = Set + PropertyName.ToString();
+		const FString FunctionString = Set + PropertyName.ToString();
 
-	FunctionName = FName(*FunctionString);
+		FunctionName = FName(*FunctionString);
+	}
 }
 
 void FTrackInstancePropertyBindings::CallFunction( UObject* InRuntimeObject, void* FunctionParams )
diff --git a/Engine/Source/Runtime/MovieSceneCore/Public/MovieSceneCommonHelpers.h b/Engine/Source/Runtime/MovieSceneCore/Public/MovieSceneCommonHelpers.h
--- a/Engine/Source/Runtime/MovieSceneCore/Public/MovieSceneCommonHelpers.h
+++ b/Engine/Source/Runtime/MovieSceneCore/Public/MovieSceneCommonHelpers.h
@@ -19,6 +19,22 @@ public:
 	*/
 	static UMovieSceneSection* FindSectionAtTime( const TArray<UMovieSceneSection*>& Sections, float Time );
 
+	/**
+	* Finds the nearest section to the given time
+	*
+	* @param Time	The time to find a section at
+	* @return The found section or null
+	*/
+	static UMovieSceneSection* FindNearestSectionAtTime( const TArray<UMovieSceneSection*>& Sections, float Time );
+
+	/*
+	* Gather all key indices bounded within a section 
+	*
+	* @param Section The section to gather keys for
+	* @param Curve The curve to gather keys for
+	* @param KeyIndices The set of indices
+	*/
+	static void GatherAllKeyIndicesInSection(const UMovieSceneSection* Section, const FIndexedCurve& Curve, TSet<int32>& KeyIndices);
 };
 
 /**
@@ -28,7 +44,7 @@ public:
 class MOVIESCENECORE_API FTrackInstancePropertyBindings
 {
 public:
-	FTrackInstancePropertyBindings( FName InPropertyName, const FString& InPropertyPath );
+	FTrackInstancePropertyBindings( FName InPropertyName, const FString& InPropertyPath, const FName& InFunctionName = FName());
 
 	/**
 	 * Calls the setter function for a specific runtime object
diff --git a/Engine/Source/Runtime/MovieSceneCoreTypes/Private/MovieSceneBoolTrack.cpp b/Engine/Source/Runtime/MovieSceneCoreTypes/Private/MovieSceneBoolTrack.cpp
--- a/Engine/Source/Runtime/MovieSceneCoreTypes/Private/MovieSceneBoolTrack.cpp
+++ b/Engine/Source/Runtime/MovieSceneCoreTypes/Private/MovieSceneBoolTrack.cpp
@@ -24,7 +24,7 @@ TSharedPtr<IMovieSceneTrackInstance> UMovieSceneBoolTrack::CreateInstance()
 
 bool UMovieSceneBoolTrack::AddKeyToSection( float Time, bool Value )
 {
-	const UMovieSceneSection* NearestSection = MovieSceneHelpers::FindSectionAtTime( Sections, Time );
+	const UMovieSceneSection* NearestSection = MovieSceneHelpers::FindNearestSectionAtTime( Sections, Time );
 	if (!NearestSection || CastChecked<UMovieSceneBoolSection>(NearestSection)->NewKeyIsNewData(Time, Value))
 	{
 		Modify();
diff --git a/Engine/Source/Runtime/MovieSceneCoreTypes/Private/MovieSceneByteTrack.cpp b/Engine/Source/Runtime/MovieSceneCoreTypes/Private/MovieSceneByteTrack.cpp
--- a/Engine/Source/Runtime/MovieSceneCoreTypes/Private/MovieSceneByteTrack.cpp
+++ b/Engine/Source/Runtime/MovieSceneCoreTypes/Private/MovieSceneByteTrack.cpp
@@ -24,7 +24,7 @@ TSharedPtr<IMovieSceneTrackInstance> UMovieSceneByteTrack::CreateInstance()
 
 bool UMovieSceneByteTrack::AddKeyToSection( float Time, uint8 Value )
 {
-	const UMovieSceneSection* NearestSection = MovieSceneHelpers::FindSectionAtTime( Sections, Time );
+	const UMovieSceneSection* NearestSection = MovieSceneHelpers::FindNearestSectionAtTime( Sections, Time );
 	if (!NearestSection || CastChecked<UMovieSceneByteSection>(NearestSection)->NewKeyIsNewData(Time, Value))
 	{
 		Modify();
diff --git a/Engine/Source/Runtime/MovieSceneCoreTypes/Private/MovieSceneColorTrack.cpp b/Engine/Source/Runtime/MovieSceneCoreTypes/Private/MovieSceneColorTrack.cpp
--- a/Engine/Source/Runtime/MovieSceneCoreTypes/Private/MovieSceneColorTrack.cpp
+++ b/Engine/Source/Runtime/MovieSceneCoreTypes/Private/MovieSceneColorTrack.cpp
@@ -27,7 +27,7 @@ bool UMovieSceneColorTrack::AddKeyToSection( float Time, const FColorKey& Key )
 {
 	bIsSlateColor = Key.bIsSlateColor;
 
-	const UMovieSceneSection* NearestSection = MovieSceneHelpers::FindSectionAtTime( Sections, Time );
+	const UMovieSceneSection* NearestSection = MovieSceneHelpers::FindNearestSectionAtTime( Sections, Time );
 	if (!NearestSection || Key.bAddKeyEvenIfUnchanged || CastChecked<UMovieSceneColorSection>(NearestSection)->NewKeyIsNewData(Time, Key.Value))
 	{
 		Modify();
diff --git a/Engine/Source/Runtime/MovieSceneCoreTypes/Private/MovieSceneFloatTrack.cpp b/Engine/Source/Runtime/MovieSceneCoreTypes/Private/MovieSceneFloatTrack.cpp
--- a/Engine/Source/Runtime/MovieSceneCoreTypes/Private/MovieSceneFloatTrack.cpp
+++ b/Engine/Source/Runtime/MovieSceneCoreTypes/Private/MovieSceneFloatTrack.cpp
@@ -25,7 +25,7 @@ TSharedPtr<IMovieSceneTrackInstance> UMovieSceneFloatTrack::CreateInstance()
 
 bool UMovieSceneFloatTrack::AddKeyToSection( float Time, float Value )
 {
-	const UMovieSceneSection* NearestSection = MovieSceneHelpers::FindSectionAtTime( Sections, Time );
+	const UMovieSceneSection* NearestSection = MovieSceneHelpers::FindNearestSectionAtTime( Sections, Time );
 	if (!NearestSection || CastChecked<UMovieSceneFloatSection>(NearestSection)->NewKeyIsNewData(Time, Value))
 	{
 		Modify();
diff --git a/Engine/Source/Runtime/MovieSceneCoreTypes/Private/MovieSceneVectorTrack.cpp b/Engine/Source/Runtime/MovieSceneCoreTypes/Private/MovieSceneVectorTrack.cpp
--- a/Engine/Source/Runtime/MovieSceneCoreTypes/Private/MovieSceneVectorTrack.cpp
+++ b/Engine/Source/Runtime/MovieSceneCoreTypes/Private/MovieSceneVectorTrack.cpp
@@ -24,7 +24,7 @@ TSharedPtr<IMovieSceneTrackInstance> UMovieSceneVectorTrack::CreateInstance()
 
 bool UMovieSceneVectorTrack::AddKeyToSection( float Time, const FVector4& Value, int32 InChannelsUsed, FName CurveName, bool bAddKeyEvenIfUnchanged )
 {
-	const UMovieSceneSection* NearestSection = MovieSceneHelpers::FindSectionAtTime( Sections, Time );
+	const UMovieSceneSection* NearestSection = MovieSceneHelpers::FindNearestSectionAtTime( Sections, Time );
 	if (!NearestSection || bAddKeyEvenIfUnchanged || CastChecked<UMovieSceneVectorSection>(NearestSection)->NewKeyIsNewData(Time, Value))
 	{
 		Modify();
diff --git a/Engine/Source/Runtime/UMG/Private/Animation/MovieSceneMarginTrack.cpp b/Engine/Source/Runtime/UMG/Private/Animation/MovieSceneMarginTrack.cpp
--- a/Engine/Source/Runtime/UMG/Private/Animation/MovieSceneMarginTrack.cpp
+++ b/Engine/Source/Runtime/UMG/Private/Animation/MovieSceneMarginTrack.cpp
@@ -25,7 +25,7 @@ TSharedPtr<IMovieSceneTrackInstance> UMovieSceneMarginTrack::CreateInstance()
 
 bool UMovieSceneMarginTrack::AddKeyToSection( float Time, const FMarginKey& MarginKey )
 {
-	const UMovieSceneSection* NearestSection = MovieSceneHelpers::FindSectionAtTime( Sections, Time );
+	const UMovieSceneSection* NearestSection = MovieSceneHelpers::FindNearestSectionAtTime( Sections, Time );
 	if (!NearestSection || MarginKey.bAddKeyEvenIfUnchanged || CastChecked<UMovieSceneMarginSection>(NearestSection)->NewKeyIsNewData(Time, MarginKey.Value) )
 	{
 		Modify();
