Sequencer: Remove unused function.

diff --git a/Engine/Source/Runtime/MovieSceneCore/Private/MovieSceneCommonHelpers.cpp b/Engine/Source/Runtime/MovieSceneCore/Private/MovieSceneCommonHelpers.cpp
--- a/Engine/Source/Runtime/MovieSceneCore/Private/MovieSceneCommonHelpers.cpp
+++ b/Engine/Source/Runtime/MovieSceneCore/Private/MovieSceneCommonHelpers.cpp
@@ -89,22 +89,6 @@ UMovieSceneSection* MovieSceneHelpers::FindNearestSectionAtTime( const TArray<UM
 	return ClosestSection ? ClosestSection : EarliestSection;
 }
 
-void MovieSceneHelpers::GatherAllKeyIndicesInSection(const UMovieSceneSection* Section, const FIndexedCurve& Curve, TSet<int32>& KeyIndices)
-{
-	KeyIndices.Empty();
-
-	float StartTime = Section->GetStartTime();
-	float EndTime = Section->GetEndTime();
-
-	///@todo this probably isn't efficient?
-	for (auto It(Curve.GetKeyHandleIterator()); It; ++It)
-	{
-		float Time = Curve.GetKeyTime(It.Key());
-		if (Time >= StartTime && Time <= EndTime)
-			KeyIndices.Add(It.Value());
-	}
-}
-
 FTrackInstancePropertyBindings::FTrackInstancePropertyBindings( FName InPropertyName, const FString& InPropertyPath, const FName& InFunctionName )
     : PropertyPath( InPropertyPath )
 	, PropertyName( InPropertyName )
diff --git a/Engine/Source/Runtime/MovieSceneCore/Public/MovieSceneCommonHelpers.h b/Engine/Source/Runtime/MovieSceneCore/Public/MovieSceneCommonHelpers.h
--- a/Engine/Source/Runtime/MovieSceneCore/Public/MovieSceneCommonHelpers.h
+++ b/Engine/Source/Runtime/MovieSceneCore/Public/MovieSceneCommonHelpers.h
@@ -26,15 +26,6 @@ public:
 	* @return The found section or null
 	*/
 	static UMovieSceneSection* FindNearestSectionAtTime( const TArray<UMovieSceneSection*>& Sections, float Time );
-
-	/*
-	* Gather all key indices bounded within a section 
-	*
-	* @param Section The section to gather keys for
-	* @param Curve The curve to gather keys for
-	* @param KeyIndices The set of indices
-	*/
-	static void GatherAllKeyIndicesInSection(const UMovieSceneSection* Section, const FIndexedCurve& Curve, TSet<int32>& KeyIndices);
 };
 
 /**
